Engine Cameras: API for EngineCamerasSubsystem

diff --git a/Engine/Plugins/Cameras/EngineCameras/Source/EngineCameras/Public/EngineCamerasSubsystem.h b/Engine/Plugins/Cameras/EngineCameras/Source/EngineCameras/Public/EngineCamerasSubsystem.h
--- a/Engine/Plugins/Cameras/EngineCameras/Source/EngineCameras/Public/EngineCamerasSubsystem.h
+++ b/Engine/Plugins/Cameras/EngineCameras/Source/EngineCameras/Public/EngineCamerasSubsystem.h
@@ -13,7 +13,7 @@ struct FCameraAnimationParams;
  * World subsystem that holds global objects for handling camera animation sequences.
  */
 UCLASS()
-class UEngineCamerasSubsystem : public UWorldSubsystem
+class ENGINECAMERAS_API UEngineCamerasSubsystem : public UWorldSubsystem
 {
 	GENERATED_BODY()
 
