Don't restart a camera animation blend out when reaching the blend out time if we were already stopping.

PR #12576

diff --git a/Engine/Plugins/Cameras/EngineCameras/Source/EngineCameras/Private/Animations/CameraAnimationCameraModifier.cpp b/Engine/Plugins/Cameras/EngineCameras/Source/EngineCameras/Private/Animations/CameraAnimationCameraModifier.cpp
--- a/Engine/Plugins/Cameras/EngineCameras/Source/EngineCameras/Private/Animations/CameraAnimationCameraModifier.cpp
+++ b/Engine/Plugins/Cameras/EngineCameras/Source/EngineCameras/Private/Animations/CameraAnimationCameraModifier.cpp
@@ -261,6 +261,8 @@ int32 UCameraAnimationCameraModifier::FindInactiveCameraAnimation()
 		const FActiveCameraAnimationInfo& CameraAnimation(ActiveAnimations[Index]);
 		if (!CameraAnimation.IsValid())
 		{
+			// Reset this entry to default.
+			ActiveAnimations[Index] = FActiveCameraAnimationInfo();
 			return Index;
 		}
 	}
@@ -363,12 +365,14 @@ void UCameraAnimationCameraModifier::TickAnimation(FActiveCameraAnimationInfo& C
 	}
 
 	// Start easing out if we're nearing the end.
+	// We might already be easing out if we have been told to stop. In that case, don't reset
+	// the ease out, just continue doing it.
 	if (!Player->GetIsLooping())
 	{
 		const float ElapsedTime = DisplayRate.AsSeconds(NewPosition - Player->GetStartFrame());
 		const float DurationTime = DisplayRate.AsSeconds(Player->GetDuration()) * Params.PlayRate;
 		const float BlendOutStartTime = DurationTime - Params.EaseOutDuration;
-		if (ElapsedTime > BlendOutStartTime)
+		if (ElapsedTime > BlendOutStartTime && !CameraAnimation.bIsEasingOut)
 		{
 			CameraAnimation.bIsEasingOut = true;
 			CameraAnimation.EaseOutCurrentTime = ElapsedTime - BlendOutStartTime;
